21219961 UE-160163 Eye Tracking additional features
-Basic support for blink and pupil info in eye tracking.  Will be false/0.0 if this feature is not supported by the platform.
-Fixed a typo in Epic_openxr.h

diff --git a/Engine/Plugins/Runtime/OpenXR/Source/OpenXRHMD/Private/Epic_openxr.h b/Engine/Plugins/Runtime/OpenXR/Source/OpenXRHMD/Private/Epic_openxr.h
--- a/Engine/Plugins/Runtime/OpenXR/Source/OpenXRHMD/Private/Epic_openxr.h
+++ b/Engine/Plugins/Runtime/OpenXR/Source/OpenXRHMD/Private/Epic_openxr.h
@@ -18,7 +18,7 @@ typedef enum XrStructureTypeEPIC {
 
 //  Extensions
 
-#define XR_EPIC_space_acceleration_1
+#define XR_EPIC_space_acceleration 1
 #define XR_EPIC_space_acceleration_SPEC_VERSION 1
 #define XR_EPIC_SPACE_ACCELERATION_NAME "XR_EPIC_space_acceleration"
 
diff --git a/Engine/Source/Runtime/EyeTracker/Public/EyeTrackerTypes.h b/Engine/Source/Runtime/EyeTracker/Public/EyeTrackerTypes.h
--- a/Engine/Source/Runtime/EyeTracker/Public/EyeTrackerTypes.h
+++ b/Engine/Source/Runtime/EyeTracker/Public/EyeTrackerTypes.h
@@ -31,6 +31,10 @@ public:
 		, GazeDirection(ForceInitToZero)
 		, FixationPoint(ForceInitToZero)
 		, ConfidenceValue(0.f)
+		, bIsLeftEyeBlink(false)
+		, bIsRightEyeBlink(false)
+		, LeftPupilDiameter(0.0f)
+		, RightPupilDiameter(0.0f)
 	{}
 
 	/** Origin of the unified gaze ray. */
@@ -48,6 +52,22 @@ public:
 	/** Value [0..1] that represents confidence in the gaze ray data above. Stable, reliably tracked eyes will be at or near 1. Fast-moving or unreliably tracked eyes will be at or near 0. */
 	UPROPERTY(EditAnywhere, BlueprintReadWrite, Category = "Gaze Data")
 	float ConfidenceValue;
+
+	/** Left eye blink data. True if the eye is close. False if the eye is open.*/
+	UPROPERTY(EditAnywhere, BlueprintReadWrite, Category = "Gaze Data")
+	bool bIsLeftEyeBlink;
+
+	/** Right eye blink data. True if the eye is close. False if the eye is open. */
+	UPROPERTY(EditAnywhere, BlueprintReadWrite, Category = "Gaze Data")
+	bool bIsRightEyeBlink;
+
+	/** Diameter of the left pupil. */
+	UPROPERTY(EditAnywhere, BlueprintReadWrite, Category = "Gaze Data")
+	float LeftPupilDiameter;
+
+	/** Diameter of the right pupil */
+	UPROPERTY(EditAnywhere, BlueprintReadWrite, Category = "Gaze Data")
+	float RightPupilDiameter;
 };
 
 
