[EyeTracker]
* Dll exported IEyeTrackerModule::GetModularFeatureName() to work around issue with blueprint functions not appearing. Have no explanation why this happens but based on that EyeTrackerModule is sub classed and two module instances (base and sub class) will call RegisterModularFeature with the same name

(If I were to guess something it would be static initialization ordering and inlined function calls)

diff --git a/Engine/Source/Runtime/EyeTracker/Private/EyeTrackerModule.cpp b/Engine/Source/Runtime/EyeTracker/Private/EyeTrackerModule.cpp
--- a/Engine/Source/Runtime/EyeTracker/Private/EyeTrackerModule.cpp
+++ b/Engine/Source/Runtime/EyeTracker/Private/EyeTrackerModule.cpp
@@ -3,6 +3,12 @@
 #include "IEyeTrackerModule.h"
 #include "Modules/ModuleManager.h"
 
+FName IEyeTrackerModule::GetModularFeatureName()
+{
+	static const FName EyeTrackerModularFeatureName(TEXT("EyeTracker"));
+	return EyeTrackerModularFeatureName;
+}
+
 
 class FEyeTrackerModule : public IEyeTrackerModule
 {
diff --git a/Engine/Source/Runtime/EyeTracker/Public/IEyeTrackerModule.h b/Engine/Source/Runtime/EyeTracker/Public/IEyeTrackerModule.h
--- a/Engine/Source/Runtime/EyeTracker/Public/IEyeTrackerModule.h
+++ b/Engine/Source/Runtime/EyeTracker/Public/IEyeTrackerModule.h
@@ -15,11 +15,7 @@ class IEyeTrackerModule : public IModuleInterface, public IModularFeature
 {
 public:
 	/** Returns modular feature name for this module */
-	static FName GetModularFeatureName()
-	{
-		static const FName EyeTrackerModularFeatureName(TEXT("EyeTracker"));
-		return EyeTrackerModularFeatureName;
-	}
+	static EYETRACKER_API FName GetModularFeatureName();
 
 	virtual FString GetModuleKeyName() const = 0;
 
